Updating TP_InCamVFXBP to avoid compilation error due to incompatible version of Visual Studio installed on the host

diff --git a/Templates/TP_InCamVFXBP/Config/DefaultEngine.ini b/Templates/TP_InCamVFXBP/Config/DefaultEngine.ini
--- a/Templates/TP_InCamVFXBP/Config/DefaultEngine.ini
+++ b/Templates/TP_InCamVFXBP/Config/DefaultEngine.ini
@@ -25,7 +25,7 @@ DefaultGraphicsRHI=DefaultGraphicsRHI_DX12
 +D3D12TargetedShaderFormats=PCD3D_SM6
 -D3D11TargetedShaderFormats=PCD3D_SM5
 +D3D11TargetedShaderFormats=PCD3D_SM5
-Compiler=VisualStudio2019
+Compiler=Default
 bTarget32Bit=False
 AudioSampleRate=48000
 AudioCallbackBufferFrameSize=1024
