TP_InCamVFXPBP: Removed dependencies that prevented usage of the template in Linux.

Also allowed it to be shipped for Linux.

#virtualized

diff --git a/Engine/Build/InstalledEngineFilters.xml b/Engine/Build/InstalledEngineFilters.xml
--- a/Engine/Build/InstalledEngineFilters.xml
+++ b/Engine/Build/InstalledEngineFilters.xml
@@ -414,6 +414,8 @@
 
 		<!-- Virtual Production Templates -->
 		Templates/TP_nDisplayBP/...
+		Templates/TP_InCamVFXBP/...
+
 
 		<!-- DotNet -->
 		Engine/Binaries/ThirdParty/DotNet/6.0.302/linux/...
diff --git a/Templates/TP_InCamVFXBP/Config/DefaultGame.ini b/Templates/TP_InCamVFXBP/Config/DefaultGame.ini
--- a/Templates/TP_InCamVFXBP/Config/DefaultGame.ini
+++ b/Templates/TP_InCamVFXBP/Config/DefaultGame.ini
@@ -5,5 +5,5 @@ bUseBorderlessWindow=True
 [/Script/UnrealEd.ProjectPackagingSettings]
 +DirectoriesToAlwaysStageAsNonUFS=(Path="ExampleConfigs")
 
-
-
+[Staging]
++AllowedConfigFiles=TP_InCamVFXBP/Config/Tags/VPRoles.ini
diff --git a/Templates/TP_InCamVFXBP/Config/DefaultInput.ini b/Templates/TP_InCamVFXBP/Config/DefaultInput.ini
--- a/Templates/TP_InCamVFXBP/Config/DefaultInput.ini
+++ b/Templates/TP_InCamVFXBP/Config/DefaultInput.ini
@@ -1,4 +1,5 @@
 
+
 [/Script/Engine.InputSettings]
 -AxisConfig=(AxisKeyName="Gamepad_LeftX",AxisProperties=(DeadZone=0.25,Exponent=1.f,Sensitivity=1.f))
 -AxisConfig=(AxisKeyName="Gamepad_LeftY",AxisProperties=(DeadZone=0.25,Exponent=1.f,Sensitivity=1.f))
@@ -6,6 +7,8 @@
 -AxisConfig=(AxisKeyName="Gamepad_RightY",AxisProperties=(DeadZone=0.25,Exponent=1.f,Sensitivity=1.f))
 -AxisConfig=(AxisKeyName="MouseX",AxisProperties=(DeadZone=0.f,Exponent=1.f,Sensitivity=0.07f))
 -AxisConfig=(AxisKeyName="MouseY",AxisProperties=(DeadZone=0.f,Exponent=1.f,Sensitivity=0.07f))
+-AxisConfig=(AxisKeyName="Mouse2D",AxisProperties=(DeadZone=0.f,Exponent=1.f,Sensitivity=0.07f))
++AxisConfig=(AxisKeyName="Mouse2D",AxisProperties=(DeadZone=0.000000,Sensitivity=0.070000,Exponent=1.000000,bInvert=False))
 +AxisConfig=(AxisKeyName="Gamepad_LeftX",AxisProperties=(DeadZone=0.250000,Sensitivity=1.000000,Exponent=1.000000,bInvert=False))
 +AxisConfig=(AxisKeyName="Gamepad_LeftY",AxisProperties=(DeadZone=0.250000,Sensitivity=1.000000,Exponent=1.000000,bInvert=False))
 +AxisConfig=(AxisKeyName="Gamepad_RightX",AxisProperties=(DeadZone=0.250000,Sensitivity=1.000000,Exponent=1.000000,bInvert=False))
@@ -27,16 +30,62 @@
 +AxisConfig=(AxisKeyName="MotionController_Right_Grip2Axis",AxisProperties=(DeadZone=0.000000,Sensitivity=1.000000,Exponent=1.000000,bInvert=False))
 +AxisConfig=(AxisKeyName="Gamepad_Special_Left_X",AxisProperties=(DeadZone=0.000000,Sensitivity=1.000000,Exponent=1.000000,bInvert=False))
 +AxisConfig=(AxisKeyName="Gamepad_Special_Left_Y",AxisProperties=(DeadZone=0.000000,Sensitivity=1.000000,Exponent=1.000000,bInvert=False))
++AxisConfig=(AxisKeyName="Vive_Left_Trigger_Axis",AxisProperties=(DeadZone=0.000000,Sensitivity=1.000000,Exponent=1.000000,bInvert=False))
++AxisConfig=(AxisKeyName="Vive_Left_Trackpad_X",AxisProperties=(DeadZone=0.000000,Sensitivity=1.000000,Exponent=1.000000,bInvert=False))
++AxisConfig=(AxisKeyName="Vive_Left_Trackpad_Y",AxisProperties=(DeadZone=0.000000,Sensitivity=1.000000,Exponent=1.000000,bInvert=False))
++AxisConfig=(AxisKeyName="Vive_Right_Trigger_Axis",AxisProperties=(DeadZone=0.000000,Sensitivity=1.000000,Exponent=1.000000,bInvert=False))
++AxisConfig=(AxisKeyName="Vive_Right_Trackpad_X",AxisProperties=(DeadZone=0.000000,Sensitivity=1.000000,Exponent=1.000000,bInvert=False))
++AxisConfig=(AxisKeyName="Vive_Right_Trackpad_Y",AxisProperties=(DeadZone=0.000000,Sensitivity=1.000000,Exponent=1.000000,bInvert=False))
++AxisConfig=(AxisKeyName="MixedReality_Left_Trigger_Axis",AxisProperties=(DeadZone=0.000000,Sensitivity=1.000000,Exponent=1.000000,bInvert=False))
++AxisConfig=(AxisKeyName="MixedReality_Left_Thumbstick_X",AxisProperties=(DeadZone=0.000000,Sensitivity=1.000000,Exponent=1.000000,bInvert=False))
++AxisConfig=(AxisKeyName="MixedReality_Left_Thumbstick_Y",AxisProperties=(DeadZone=0.000000,Sensitivity=1.000000,Exponent=1.000000,bInvert=False))
++AxisConfig=(AxisKeyName="MixedReality_Left_Trackpad_X",AxisProperties=(DeadZone=0.000000,Sensitivity=1.000000,Exponent=1.000000,bInvert=False))
++AxisConfig=(AxisKeyName="MixedReality_Left_Trackpad_Y",AxisProperties=(DeadZone=0.000000,Sensitivity=1.000000,Exponent=1.000000,bInvert=False))
++AxisConfig=(AxisKeyName="MixedReality_Right_Trigger_Axis",AxisProperties=(DeadZone=0.000000,Sensitivity=1.000000,Exponent=1.000000,bInvert=False))
++AxisConfig=(AxisKeyName="MixedReality_Right_Thumbstick_X",AxisProperties=(DeadZone=0.000000,Sensitivity=1.000000,Exponent=1.000000,bInvert=False))
++AxisConfig=(AxisKeyName="MixedReality_Right_Thumbstick_Y",AxisProperties=(DeadZone=0.000000,Sensitivity=1.000000,Exponent=1.000000,bInvert=False))
++AxisConfig=(AxisKeyName="MixedReality_Right_Trackpad_X",AxisProperties=(DeadZone=0.000000,Sensitivity=1.000000,Exponent=1.000000,bInvert=False))
++AxisConfig=(AxisKeyName="MixedReality_Right_Trackpad_Y",AxisProperties=(DeadZone=0.000000,Sensitivity=1.000000,Exponent=1.000000,bInvert=False))
++AxisConfig=(AxisKeyName="OculusTouch_Left_Grip_Axis",AxisProperties=(DeadZone=0.000000,Sensitivity=1.000000,Exponent=1.000000,bInvert=False))
++AxisConfig=(AxisKeyName="OculusTouch_Left_Trigger_Axis",AxisProperties=(DeadZone=0.000000,Sensitivity=1.000000,Exponent=1.000000,bInvert=False))
++AxisConfig=(AxisKeyName="OculusTouch_Left_Thumbstick_X",AxisProperties=(DeadZone=0.000000,Sensitivity=1.000000,Exponent=1.000000,bInvert=False))
++AxisConfig=(AxisKeyName="OculusTouch_Left_Thumbstick_Y",AxisProperties=(DeadZone=0.000000,Sensitivity=1.000000,Exponent=1.000000,bInvert=False))
++AxisConfig=(AxisKeyName="OculusTouch_Right_Grip_Axis",AxisProperties=(DeadZone=0.000000,Sensitivity=1.000000,Exponent=1.000000,bInvert=False))
++AxisConfig=(AxisKeyName="OculusTouch_Right_Trigger_Axis",AxisProperties=(DeadZone=0.000000,Sensitivity=1.000000,Exponent=1.000000,bInvert=False))
++AxisConfig=(AxisKeyName="OculusTouch_Right_Thumbstick_X",AxisProperties=(DeadZone=0.000000,Sensitivity=1.000000,Exponent=1.000000,bInvert=False))
++AxisConfig=(AxisKeyName="OculusTouch_Right_Thumbstick_Y",AxisProperties=(DeadZone=0.000000,Sensitivity=1.000000,Exponent=1.000000,bInvert=False))
++AxisConfig=(AxisKeyName="ValveIndex_Left_Grip_Axis",AxisProperties=(DeadZone=0.000000,Sensitivity=1.000000,Exponent=1.000000,bInvert=False))
++AxisConfig=(AxisKeyName="ValveIndex_Left_Grip_Force",AxisProperties=(DeadZone=0.000000,Sensitivity=1.000000,Exponent=1.000000,bInvert=False))
++AxisConfig=(AxisKeyName="ValveIndex_Left_Trigger_Axis",AxisProperties=(DeadZone=0.000000,Sensitivity=1.000000,Exponent=1.000000,bInvert=False))
++AxisConfig=(AxisKeyName="ValveIndex_Left_Thumbstick_X",AxisProperties=(DeadZone=0.000000,Sensitivity=1.000000,Exponent=1.000000,bInvert=False))
++AxisConfig=(AxisKeyName="ValveIndex_Left_Thumbstick_Y",AxisProperties=(DeadZone=0.000000,Sensitivity=1.000000,Exponent=1.000000,bInvert=False))
++AxisConfig=(AxisKeyName="ValveIndex_Left_Trackpad_X",AxisProperties=(DeadZone=0.000000,Sensitivity=1.000000,Exponent=1.000000,bInvert=False))
++AxisConfig=(AxisKeyName="ValveIndex_Left_Trackpad_Y",AxisProperties=(DeadZone=0.000000,Sensitivity=1.000000,Exponent=1.000000,bInvert=False))
++AxisConfig=(AxisKeyName="ValveIndex_Left_Trackpad_Force",AxisProperties=(DeadZone=0.000000,Sensitivity=1.000000,Exponent=1.000000,bInvert=False))
++AxisConfig=(AxisKeyName="ValveIndex_Right_Grip_Axis",AxisProperties=(DeadZone=0.000000,Sensitivity=1.000000,Exponent=1.000000,bInvert=False))
++AxisConfig=(AxisKeyName="ValveIndex_Right_Grip_Force",AxisProperties=(DeadZone=0.000000,Sensitivity=1.000000,Exponent=1.000000,bInvert=False))
++AxisConfig=(AxisKeyName="ValveIndex_Right_Trigger_Axis",AxisProperties=(DeadZone=0.000000,Sensitivity=1.000000,Exponent=1.000000,bInvert=False))
++AxisConfig=(AxisKeyName="ValveIndex_Right_Thumbstick_X",AxisProperties=(DeadZone=0.000000,Sensitivity=1.000000,Exponent=1.000000,bInvert=False))
++AxisConfig=(AxisKeyName="ValveIndex_Right_Thumbstick_Y",AxisProperties=(DeadZone=0.000000,Sensitivity=1.000000,Exponent=1.000000,bInvert=False))
++AxisConfig=(AxisKeyName="ValveIndex_Right_Trackpad_X",AxisProperties=(DeadZone=0.000000,Sensitivity=1.000000,Exponent=1.000000,bInvert=False))
++AxisConfig=(AxisKeyName="ValveIndex_Right_Trackpad_Y",AxisProperties=(DeadZone=0.000000,Sensitivity=1.000000,Exponent=1.000000,bInvert=False))
++AxisConfig=(AxisKeyName="ValveIndex_Right_Trackpad_Force",AxisProperties=(DeadZone=0.000000,Sensitivity=1.000000,Exponent=1.000000,bInvert=False))
 bAltEnterTogglesFullscreen=True
 bF11TogglesFullscreen=True
 bUseMouseForTouch=False
 bEnableMouseSmoothing=True
 bEnableFOVScaling=True
 bCaptureMouseOnLaunch=True
-bDefaultViewportMouseLock=False
+bEnableLegacyInputScales=True
+bEnableMotionControls=True
+bFilterInputByPlatformUser=False
+bEnableInputDeviceSubsystem=True
+bShouldFlushPressedKeysOnViewportFocusLost=True
+bEnableDynamicComponentInputBinding=True
 bAlwaysShowTouchInterface=False
 bShowConsoleOnFourFingerTap=True
 bEnableGestureRecognizer=False
+bUseAutocorrect=False
 DefaultViewportMouseCaptureMode=CapturePermanently_IncludingInitialMouseDown
 DefaultViewportMouseLockMode=LockOnCapture
 FOVScale=0.011110
@@ -53,9 +102,9 @@ DoubleClickTime=0.200000
 +AxisMappings=(AxisName="TurnRate",Scale=1.000000,Key=MouseX)
 +AxisMappings=(AxisName="LookUpRate",Scale=1.000000,Key=Down)
 +AxisMappings=(AxisName="LookUpRate",Scale=-1.000000,Key=MouseY)
+DefaultPlayerInputClass=/Script/EnhancedInput.EnhancedPlayerInput
+DefaultInputComponentClass=/Script/EnhancedInput.EnhancedInputComponent
 DefaultTouchInterface=/Engine/MobileResources/HUD/DefaultVirtualJoysticks.DefaultVirtualJoysticks
-ConsoleKey=None
 -ConsoleKeys=Tilde
 +ConsoleKeys=Tilde
 
-
diff --git a/Templates/TP_InCamVFXBP/TP_InCamVFXBP.uproject b/Templates/TP_InCamVFXBP/TP_InCamVFXBP.uproject
--- a/Templates/TP_InCamVFXBP/TP_InCamVFXBP.uproject
+++ b/Templates/TP_InCamVFXBP/TP_InCamVFXBP.uproject
@@ -5,33 +5,114 @@
 	"Description": "",
 	"Plugins": [
 		{
-			"Name": "ICVFX",
+			"Name": "DisplayClusterLaunch",
 			"Enabled": true
 		},
 		{
-			"Name": "OpenXR",
-			"Enabled": true,
-			"SupportedTargetPlatforms": [
-				"Win64",
-				"Linux",
-				"Android"
-			]
+			"Name": "BlackmagicMedia",
+			"Enabled": true
 		},
 		{
-			"Name": "VirtualScouting",
-			"Enabled": true,
-			"SupportedTargetPlatforms": [
-				"Win64",
-				"Linux"
-			]
+			"Name": "CameraCalibration",
+			"Enabled": true
 		},
 		{
-			"Name": "DisplayClusterLaunch",
+			"Name": "Composure",
+			"Enabled": true
+		},
+		{
+			"Name": "ColorCorrectRegions",
+			"Enabled": true
+		},
+		{
+			"Name": "LevelSnapshots",
+			"Enabled": true
+		},
+		{
+			"Name": "LiveLink",
+			"Enabled": true
+		},
+		{
+			"Name": "LiveLinkOverNDisplay",
+			"Enabled": true
+		},
+		{
+			"Name": "LiveLinkLens",
+			"Enabled": true
+		},
+		{
+			"Name": "LiveLinkCamera",
+			"Enabled": true
+		},
+		{
+			"Name": "MediaFrameworkUtilities",
+			"Enabled": true
+		},
+		{
+			"Name": "MultiUserClient",
+			"Enabled": true
+		},
+		{
+			"Name": "MultiUserTakes",
+			"Enabled": true
+		},
+		{
+			"Name": "nDisplay",
+			"Enabled": true
+		},
+		{
+			"Name": "OSC",
+			"Enabled": true
+		},
+		{
+			"Name": "OpenColorIO",
+			"Enabled": true
+		},
+		{
+			"Name": "RemoteControl",
+			"Enabled": true
+		},
+		{
+			"Name": "RemoteControlWebInterface",
+			"Enabled": true
+		},
+		{
+			"Name": "StageMonitoring",
+			"Enabled": true
+		},
+		{
+			"Name": "Switchboard",
+			"Enabled": true
+		},
+		{
+			"Name": "Takes",
+			"Enabled": true
+		},
+		{
+			"Name": "TimedDataMonitor",
+			"Enabled": true
+		},
+		{
+			"Name": "SequencerScripting",
+			"Enabled": true
+		},
+		{
+			"Name": "ConsoleVariables",
+			"Enabled": true
+		},
+		{
+			"Name": "VirtualProductionUtilities",
+			"Enabled": true
+		},
+		{
+			"Name": "EpicStageApp",
 			"Enabled": true
 		}
 	],
 	"TargetPlatforms": [
-		"WindowsNoEditor",
-		"Windows"
+		"Windows",
+		"WindowsEditor",
+		"Linux",
+		"LinuxEditor"
 	]
 }
\ No newline at end of file
