ME Templates and Showcases: Updated to use SM6 on Linux for rayracing

diff --git a/Templates/TP_DMXBP/Config/DefaultEngine.ini b/Templates/TP_DMXBP/Config/DefaultEngine.ini
--- a/Templates/TP_DMXBP/Config/DefaultEngine.ini
+++ b/Templates/TP_DMXBP/Config/DefaultEngine.ini
@@ -36,3 +36,7 @@ DefaultGraphicsRHI=DefaultGraphicsRHI_DX12
 +D3D12TargetedShaderFormats=PCD3D_SM6
 -D3D11TargetedShaderFormats=PCD3D_SM5
 +D3D11TargetedShaderFormats=PCD3D_SM5
+
+[/Script/LinuxTargetPlatform.LinuxTargetSettings]
+-TargetedRHIs=SF_VULKAN_SM5
++TargetedRHIs=SF_VULKAN_SM6
\ No newline at end of file
diff --git a/Templates/TP_InCamVFXBP/Config/DefaultEngine.ini b/Templates/TP_InCamVFXBP/Config/DefaultEngine.ini
--- a/Templates/TP_InCamVFXBP/Config/DefaultEngine.ini
+++ b/Templates/TP_InCamVFXBP/Config/DefaultEngine.ini
@@ -50,6 +50,10 @@ CompressionQualityModifier=0.000000
 AutoStreamingThreshold=0.000000
 SoundCueCookQualityIndex=-1
 
+[/Script/LinuxTargetPlatform.LinuxTargetSettings]
+-TargetedRHIs=SF_VULKAN_SM5
++TargetedRHIs=SF_VULKAN_SM6
+
 [/Script/Engine.Engine]
 +ActiveGameNameRedirects=(OldGameName="TP_nDisplayBP",NewGameName="/Script/InCameraVFXExample")
 +ActiveGameNameRedirects=(OldGameName="/Script/TP_nDisplayBP",NewGameName="/Script/InCameraVFXExample")
diff --git a/Templates/TP_ME_BlankBP/Config/DefaultEngine.ini b/Templates/TP_ME_BlankBP/Config/DefaultEngine.ini
--- a/Templates/TP_ME_BlankBP/Config/DefaultEngine.ini
+++ b/Templates/TP_ME_BlankBP/Config/DefaultEngine.ini
@@ -99,3 +99,7 @@ DefaultGraphicsRHI=DefaultGraphicsRHI_DX12
 +D3D12TargetedShaderFormats=PCD3D_SM6
 -D3D11TargetedShaderFormats=PCD3D_SM5
 +D3D11TargetedShaderFormats=PCD3D_SM5
+
+[/Script/LinuxTargetPlatform.LinuxTargetSettings]
+-TargetedRHIs=SF_VULKAN_SM5
++TargetedRHIs=SF_VULKAN_SM6
\ No newline at end of file
diff --git a/Templates/TP_ME_VProdBP/Config/DefaultEngine.ini b/Templates/TP_ME_VProdBP/Config/DefaultEngine.ini
--- a/Templates/TP_ME_VProdBP/Config/DefaultEngine.ini
+++ b/Templates/TP_ME_VProdBP/Config/DefaultEngine.ini
@@ -99,3 +99,7 @@ DefaultGraphicsRHI=DefaultGraphicsRHI_DX12
 +D3D12TargetedShaderFormats=PCD3D_SM6
 -D3D11TargetedShaderFormats=PCD3D_SM5
 +D3D11TargetedShaderFormats=PCD3D_SM5
+
+[/Script/LinuxTargetPlatform.LinuxTargetSettings]
+-TargetedRHIs=SF_VULKAN_SM5
++TargetedRHIs=SF_VULKAN_SM6
diff --git a/Templates/TP_nDisplayBP/Config/DefaultEngine.ini b/Templates/TP_nDisplayBP/Config/DefaultEngine.ini
--- a/Templates/TP_nDisplayBP/Config/DefaultEngine.ini
+++ b/Templates/TP_nDisplayBP/Config/DefaultEngine.ini
@@ -28,3 +28,7 @@ DefaultGraphicsRHI=DefaultGraphicsRHI_DX12
 +D3D12TargetedShaderFormats=PCD3D_SM6
 -D3D11TargetedShaderFormats=PCD3D_SM5
 +D3D11TargetedShaderFormats=PCD3D_SM5
+
+[/Script/LinuxTargetPlatform.LinuxTargetSettings]
+-TargetedRHIs=SF_VULKAN_SM5
++TargetedRHIs=SF_VULKAN_SM6
\ No newline at end of file
