Removing the Legacy Composure plugin from Template and Sample Projects and enabling the new Composite plugin where relevant.  Some projects also needed the Alpha Holdout Project Setting enabled

diff --git a/Templates/TP_InCamVFXBP/Config/DefaultEngine.ini b/Templates/TP_InCamVFXBP/Config/DefaultEngine.ini
--- a/Templates/TP_InCamVFXBP/Config/DefaultEngine.ini
+++ b/Templates/TP_InCamVFXBP/Config/DefaultEngine.ini
@@ -21,6 +21,8 @@ r.SceneRenderTargetResizeMethod=2
 r.DefaultBackBufferPixelFormat=4
 vr.AllowMotionBlurInVR=True
 r.AllowStaticLighting=False
+r.PostProcessing.PropagateAlpha=True
+r.Deferred.SupportPrimitiveAlphaHoldout=True
 
 [/Script/WindowsTargetPlatform.WindowsTargetSettings]
 DefaultGraphicsRHI=DefaultGraphicsRHI_DX12
@@ -67,3 +69,17 @@ Concert.EnableSequencePlayer=True
 Concert.EnableSequencerPlaybackSync=True
 Concert.EnableOpenRemoteSequencer=True
 
+[/Script/AndroidFileServerEditor.AndroidFileServerRuntimeSettings]
+bEnablePlugin=True
+bAllowNetworkConnection=True
+SecurityToken=CB4D93E34925E8C785C5919DD25FC1D5
+bIncludeInShipping=False
+bAllowExternalStartInShipping=False
+bCompileAFSProject=False
+bUseCompression=False
+bLogFiles=False
+bReportStats=False
+ConnectionType=USBOnly
+bUseManualIPAddress=False
+ManualIPAddress=
+
diff --git a/Templates/TP_InCamVFXBP/TP_InCamVFXBP.uproject b/Templates/TP_InCamVFXBP/TP_InCamVFXBP.uproject
--- a/Templates/TP_InCamVFXBP/TP_InCamVFXBP.uproject
+++ b/Templates/TP_InCamVFXBP/TP_InCamVFXBP.uproject
@@ -17,7 +17,7 @@
 			"Enabled": true
 		},
 		{
-			"Name": "Composure",
+			"Name": "Composite",
 			"Enabled": true
 		},
 		{
diff --git a/Templates/TP_ME_BlankBP/Config/DefaultEngine.ini b/Templates/TP_ME_BlankBP/Config/DefaultEngine.ini
--- a/Templates/TP_ME_BlankBP/Config/DefaultEngine.ini
+++ b/Templates/TP_ME_BlankBP/Config/DefaultEngine.ini
@@ -30,6 +30,7 @@ r.GBufferFormat=3
 
 ; shader settings
 r.PostProcessing.PropagateAlpha=1
+r.Deferred.SupportPrimitiveAlphaHoldout=True
 r.SupportSkyAtmosphereAffectsHeightFog=True
 r.DefaultFeature.LensFlare=True
 r.ClearCoatNormal=False
diff --git a/Templates/TP_ME_BlankBP/TP_ME_BlankBP.uproject b/Templates/TP_ME_BlankBP/TP_ME_BlankBP.uproject
--- a/Templates/TP_ME_BlankBP/TP_ME_BlankBP.uproject
+++ b/Templates/TP_ME_BlankBP/TP_ME_BlankBP.uproject
@@ -54,7 +54,7 @@
 			"Enabled": true
 		},
 		{
-			"Name": "Composure",
+			"Name": "Composite",
 			"Enabled": true
 		},
 		{
diff --git a/Templates/TP_ME_VProdBP/Config/DefaultEngine.ini b/Templates/TP_ME_VProdBP/Config/DefaultEngine.ini
--- a/Templates/TP_ME_VProdBP/Config/DefaultEngine.ini
+++ b/Templates/TP_ME_VProdBP/Config/DefaultEngine.ini
@@ -30,6 +30,7 @@ r.GBufferFormat=3
 
 ; shader settings
 r.PostProcessing.PropagateAlpha=1
+r.Deferred.SupportPrimitiveAlphaHoldout=True
 r.SupportSkyAtmosphereAffectsHeightFog=True
 r.DefaultFeature.LensFlare=False
 r.ClearCoatNormal=False
@@ -109,3 +110,18 @@ DefaultGraphicsRHI=DefaultGraphicsRHI_DX12
 [/Script/LinuxTargetPlatform.LinuxTargetSettings]
 -TargetedRHIs=SF_VULKAN_SM5
 +TargetedRHIs=SF_VULKAN_SM6
+
+[/Script/AndroidFileServerEditor.AndroidFileServerRuntimeSettings]
+bEnablePlugin=True
+bAllowNetworkConnection=True
+SecurityToken=0EAFB8DB4B4A6ADEED6F1B96CB3CA966
+bIncludeInShipping=False
+bAllowExternalStartInShipping=False
+bCompileAFSProject=False
+bUseCompression=False
+bLogFiles=False
+bReportStats=False
+ConnectionType=USBOnly
+bUseManualIPAddress=False
+ManualIPAddress=
+
diff --git a/Templates/TP_ME_VProdBP/TP_ME_VprodBP.uproject b/Templates/TP_ME_VProdBP/TP_ME_VprodBP.uproject
--- a/Templates/TP_ME_VProdBP/TP_ME_VprodBP.uproject
+++ b/Templates/TP_ME_VProdBP/TP_ME_VprodBP.uproject
@@ -54,7 +54,7 @@
 			"Enabled": true
 		},
 		{
-			"Name": "Composure",
+			"Name": "Composite",
 			"Enabled": true
 		},
 		{
diff --git a/Templates/TP_MotionDesignBP/TP_MotionDesignBP.uproject b/Templates/TP_MotionDesignBP/TP_MotionDesignBP.uproject
--- a/Templates/TP_MotionDesignBP/TP_MotionDesignBP.uproject
+++ b/Templates/TP_MotionDesignBP/TP_MotionDesignBP.uproject
@@ -33,7 +33,7 @@
 			"Enabled": true
 		},
 		{
-			"Name": "Composure",
+			"Name": "Composite",
 			"Enabled": true
 		},
 		{
