ICVFX: Revert previous config change to ICVFX template project given added alpha costs.

#rnx

diff --git a/Templates/TP_InCamVFXBP/Config/DefaultEngine.ini b/Templates/TP_InCamVFXBP/Config/DefaultEngine.ini
--- a/Templates/TP_InCamVFXBP/Config/DefaultEngine.ini
+++ b/Templates/TP_InCamVFXBP/Config/DefaultEngine.ini
@@ -21,8 +21,6 @@ r.SceneRenderTargetResizeMethod=2
 r.DefaultBackBufferPixelFormat=4
 vr.AllowMotionBlurInVR=True
 r.AllowStaticLighting=False
-r.PostProcessing.PropagateAlpha=True
-r.Deferred.SupportPrimitiveAlphaHoldout=True
 
 [/Script/WindowsTargetPlatform.WindowsTargetSettings]
 DefaultGraphicsRHI=DefaultGraphicsRHI_DX12
